Cameras: add p4ignore for binary outputs

diff --git a/Engine/Plugins/Cameras/.p4ignore.txt b/Engine/Plugins/Cameras/.p4ignore.txt
new file mode 100644
--- /dev/null
+++ b/Engine/Plugins/Cameras/.p4ignore.txt
@@ -0,0 +1 @@
+*/Binaries/*
