Cameras: ignore git files inside the plugin

diff --git a/Engine/Plugins/Cameras/.p4ignore.txt b/Engine/Plugins/Cameras/.p4ignore.txt
--- a/Engine/Plugins/Cameras/.p4ignore.txt
+++ b/Engine/Plugins/Cameras/.p4ignore.txt
@@ -1,2 +1,3 @@
 */Binaries/*
 */Intermediate/*
+*/.git
\ No newline at end of file
