Cameras: add p4 ignore

diff --git a/Engine/Plugins/Cameras/GameplayCameras/.p4ignore.txt b/Engine/Plugins/Cameras/GameplayCameras/.p4ignore.txt
new file mode 100644
--- /dev/null
+++ b/Engine/Plugins/Cameras/GameplayCameras/.p4ignore.txt
@@ -0,0 +1,3 @@
+/Binaries/
+/Intermediate/
+/.git/
\ No newline at end of file
