Cameras: add gitignore file

#Synced-CL 49781371

diff --git a/Engine/Plugins/Cameras/GameplayCameras/.gitignore b/Engine/Plugins/Cameras/GameplayCameras/.gitignore
new file mode 100644
--- /dev/null
+++ b/Engine/Plugins/Cameras/GameplayCameras/.gitignore
@@ -0,0 +1,2 @@
+Binaries/
+Intermediate/
